Make FCommonAccount::AddId go through the resolved account. This prevents a stale common account from registering itself with the common account manager.

diff --git a/Engine/Plugins/Online/OnlineFrameworkCommon/Source/OnlineFrameworkCommon/Private/OnlineFramework/CommonAccount.cpp b/Engine/Plugins/Online/OnlineFrameworkCommon/Source/OnlineFrameworkCommon/Private/OnlineFramework/CommonAccount.cpp
--- a/Engine/Plugins/Online/OnlineFrameworkCommon/Source/OnlineFrameworkCommon/Private/OnlineFramework/CommonAccount.cpp
+++ b/Engine/Plugins/Online/OnlineFrameworkCommon/Source/OnlineFrameworkCommon/Private/OnlineFramework/CommonAccount.cpp
@@ -68,6 +68,13 @@ bool FCommonAccount::AddId(FName FrameworkInstance, UE::Online::FAccountId Accou
 		ensure(ExistingAccountId == AccountId);
 		return false;
 	}
+	else if (RedirectAccount.IsValid())
+	{
+		if (FCommonAccountPtr PinnedRedirectAccount = RedirectAccount.Pin())
+		{
+			return PinnedRedirectAccount->AddId(FrameworkInstance, AccountId);
+		}
+	}
 
 	FCommonAccountManagerPtr StrongManager = Manager.Pin();
 	TOptional<FCommonConfigInstance> ConfigInstance = CommonConfig.GetFrameworkInstanceConfig(FrameworkInstance);
