Merge conflict from 45968147.
File moved in 37.50.

diff --git a/Engine/Plugins/Online/OnlineFrameworkCommon/Source/OnlineFrameworkCommon/Private/OnlineFramework/CommonAccountUtils.cpp b/Engine/Plugins/Online/OnlineFrameworkCommon/Source/OnlineFrameworkCommon/Private/OnlineFramework/CommonAccountUtils.cpp
--- a/Engine/Plugins/Online/OnlineFrameworkCommon/Source/OnlineFrameworkCommon/Private/OnlineFramework/CommonAccountUtils.cpp
+++ b/Engine/Plugins/Online/OnlineFrameworkCommon/Source/OnlineFrameworkCommon/Private/OnlineFramework/CommonAccountUtils.cpp
@@ -33,15 +33,18 @@ UE::Online::FAccountId GetAccountV2FromV1(const FCommonConfig& CommonConfig, con
 {
 	UE::Online::FAccountId AccountId;
 
-	if (UE::Online::IOnlineServicesPtr OnlineServices = CommonConfig.GetServices(FrameworkInstance))
+	if (InId.IsValid())
 	{
-		if (OnlineServices->GetServicesProvider() == UE::Online::EOnlineServices::Null && InId->GetType() == FName("Null"))
+		if (UE::Online::IOnlineServicesPtr OnlineServices = CommonConfig.GetServices(FrameworkInstance))
 		{
-			AccountId = UE::Online::FOnlineAccountIdRegistryNull::Get().FindOrAddAccountId(InId->ToString());
-		}
-		else if (UE::Online::FAuthOSSAdapterPtr AuthAdapter = OnlineServices->GetInterface<UE::Online::FAuthOSSAdapter>())
-		{
-			AccountId = AuthAdapter->GetAccountId(InId.ToSharedRef());
+			if (OnlineServices->GetServicesProvider() == UE::Online::EOnlineServices::Null && InId->GetType() == FName("Null"))
+			{
+				AccountId = UE::Online::FOnlineAccountIdRegistryNull::Get().FindOrAddAccountId(InId->ToString());
+			}
+			else if (UE::Online::FAuthOSSAdapterPtr AuthAdapter = OnlineServices->GetInterface<UE::Online::FAuthOSSAdapter>())
+			{
+				AccountId = AuthAdapter->GetAccountId(InId.ToSharedRef());
+			}
 		}
 	}
 	return AccountId;
