Use the source account id framework instance name to get the common account
+ if this was different from the adapter instance, we would register the wrong account id type for the framework instance
Add checks to CommonAccountManager that the account id being added matches the framework instance

diff --git a/Engine/Plugins/Online/OnlineFramework/Source/Party/Private/SocialToolkit.cpp b/Engine/Plugins/Online/OnlineFramework/Source/Party/Private/SocialToolkit.cpp
--- a/Engine/Plugins/Online/OnlineFramework/Source/Party/Private/SocialToolkit.cpp
+++ b/Engine/Plugins/Online/OnlineFramework/Source/Party/Private/SocialToolkit.cpp
@@ -1002,12 +1002,12 @@ void USocialToolkit::HandleServicesPresenceUpdated(const UE::Online::FPresenceUp
 
 	TSharedRef<const FUserPresence> UpdatedPresence = Event.UpdatedPresence;
 
-	const FName& PresenceFrameworkAdapterInstance = USocialManager::GetPresenceFrameworkAdapterInstance();
 	FCommonConfig CommonConfig(this);
 	
-	// We need a services instance with an Auth adapter to make this conversion
-	if (FCommonAccountPtr CommonAccountPtr = GetCommonAccountFromV2(CommonConfig, UpdatedPresence->AccountId, PresenceFrameworkAdapterInstance))
+	if (FCommonAccountPtr CommonAccountPtr = GetCommonAccountFromV2(CommonConfig, UpdatedPresence->AccountId, USocialManager::GetPresenceFrameworkInstance()))
 	{
+		// We need a services instance with an Auth adapter to make this conversion
+		const FName& PresenceFrameworkAdapterInstance = USocialManager::GetPresenceFrameworkAdapterInstance();
 		if (FUniqueNetIdPtr UserId = GetV1FromCommonAccount(CommonAccountPtr.ToSharedRef(), PresenceFrameworkAdapterInstance))
 		{
 			QUICK_SCOPE_CYCLE_COUNTER(STAT_USocialToolkit_HandlePresenceReceived);
diff --git a/Engine/Plugins/Online/OnlineFrameworkCommon/Source/OnlineFrameworkCommon/Private/OnlineFramework/CommonAccountManager.cpp b/Engine/Plugins/Online/OnlineFrameworkCommon/Source/OnlineFrameworkCommon/Private/OnlineFramework/CommonAccountManager.cpp
--- a/Engine/Plugins/Online/OnlineFrameworkCommon/Source/OnlineFrameworkCommon/Private/OnlineFramework/CommonAccountManager.cpp
+++ b/Engine/Plugins/Online/OnlineFrameworkCommon/Source/OnlineFrameworkCommon/Private/OnlineFramework/CommonAccountManager.cpp
@@ -65,7 +65,9 @@ FCommonAccountPtr FCommonAccountManager::GetAccount(UE::Online::FAccountId Accou
 FCommonAccountRef FCommonAccountManager::CreateNewAccount(UE::Online::FAccountId AccountId, FName FrameworkInstance)
 {
 	const FCommonAccountRef NewAccount = MakeShared<FCommonAccount>(FCommonAccount::FPrivateToken{}, AsShared(), CommonConfig);
-	NewAccount->AddAccountId(AccountId, GetCommonAccountIdType(CommonConfig, FrameworkInstance));
+	const FCommonAccountIdType AccountIdType = GetCommonAccountIdType(CommonConfig, FrameworkInstance);
+	check(AccountIdType.Key == AccountId.GetOnlineServicesType());
+	NewAccount->AddAccountId(AccountId, AccountIdType);
 	Accounts.Emplace(AccountId, NewAccount);
 	OnCommonAccountCreated().Broadcast(NewAccount);
 	OnCommonAccountIdAdded().Broadcast(NewAccount, AccountId);
@@ -197,6 +199,7 @@ FCommonAccountRef FCommonAccountManager::AddAccountId(const FCommonAccountRef& I
 	}
 	else
 	{
+		check(ConfigInstance.OnlineServices == AccountId.GetOnlineServicesType());
 		Account->AddAccountId(AccountId, GetCommonAccountIdType(ConfigInstance));
 	}
 
