- Fix crash when UpdateRichPresenceForServices is called on a logged out user

diff --git a/Engine/Plugins/Online/OnlineFrameworkCommon/Source/OnlineFrameworkCommon/Private/OnlineFramework/CommonAccountUtils.cpp b/Engine/Plugins/Online/OnlineFrameworkCommon/Source/OnlineFrameworkCommon/Private/OnlineFramework/CommonAccountUtils.cpp
--- a/Engine/Plugins/Online/OnlineFrameworkCommon/Source/OnlineFrameworkCommon/Private/OnlineFramework/CommonAccountUtils.cpp
+++ b/Engine/Plugins/Online/OnlineFrameworkCommon/Source/OnlineFrameworkCommon/Private/OnlineFramework/CommonAccountUtils.cpp
@@ -110,21 +110,24 @@ FName GetFirstFrameworkInstanceName(const FUniqueNetIdWrapper& InId)
 	// Be a OnlineSubsystemNull V1 id
 	// Be a V1 id where there is an OnlineServices adapter implementation available for the online subsystem type
 	FName FrameworkInstance;
-	if (InId.IsV1())
-	{
-		// Requires a common config entry with this name
-		FrameworkInstance = InId.GetV1Unsafe()->GetType();
-	}
-	else if (FOnlineFrameworkCommonModule* CommonModule = FOnlineFrameworkCommonModule::Get())
+	if (InId.IsValid())
 	{
-		UE::Online::EOnlineServices Services = InId.GetV2Unsafe().GetOnlineServicesType();
-		for (const TPair<TPair<FName, ECommonConfigContextType>, FCommonConfigInstance>& Config : CommonModule->GetConfig().FrameworkConfigs)
+		if (InId.IsV1())
 		{
-			const FCommonConfigInstance& ConfigInstance = Config.Value;
-			if (ConfigInstance.OnlineServices == Services)
+			// Requires a common config entry with this name
+			FrameworkInstance = InId.GetV1Unsafe()->GetType();
+		}
+		else if (FOnlineFrameworkCommonModule* CommonModule = FOnlineFrameworkCommonModule::Get())
+		{
+			UE::Online::EOnlineServices Services = InId.GetV2Unsafe().GetOnlineServicesType();
+			for (const TPair<TPair<FName, ECommonConfigContextType>, FCommonConfigInstance>& Config : CommonModule->GetConfig().FrameworkConfigs)
 			{
-				FrameworkInstance = Config.Key.Key;
-				break;
+				const FCommonConfigInstance& ConfigInstance = Config.Value;
+				if (ConfigInstance.OnlineServices == Services)
+				{
+					FrameworkInstance = Config.Key.Key;
+					break;
+				}
 			}
 		}
 	}
